- Niagara UI Renderer plugin
- Adds a UI renderer class which can be subclasses to render into UI widgets, currently only subclasses as a Sprite renderer

#tests create system using the new renderer, place them into a widget

#Synced-CL 52353515

diff --git a/Engine/Plugins/FX/NiagaraUIRenderer/NiagaraUIRenderer.uplugin b/Engine/Plugins/FX/NiagaraUIRenderer/NiagaraUIRenderer.uplugin
new file mode 100644
--- /dev/null
+++ b/Engine/Plugins/FX/NiagaraUIRenderer/NiagaraUIRenderer.uplugin
@@ -0,0 +1,35 @@
+{
+	"FileVersion": 3,
+	"Version": 1,
+	"VersionName": "1.0",
+	"FriendlyName": "Niagara UI Renderer",
+	"Description": "Renders Niagara CPU particle systems inside Slate/UMG widgets using a dedicated UI sprite renderer.",
+	"Category": "FX",
+	"CreatedBy": "Epic Games",
+	"CreatedByURL": "",
+	"DocsURL": "",
+	"MarketplaceURL": "",
+	"SupportURL": "",
+	"CanContainContent": true,
+	"IsBetaVersion": false,
+	"IsExperimentalVersion": true,
+	"Installed": false,
+	"Modules": [
+		{
+			"Name": "NiagaraUIRenderer",
+			"Type": "Runtime",
+			"LoadingPhase": "Default"
+		},
+		{
+			"Name": "NiagaraUIRendererEditor",
+			"Type": "Editor",
+			"LoadingPhase": "Default"
+		}
+	],
+	"Plugins": [
+		{
+			"Name": "Niagara",
+			"Enabled": true
+		}
+	]
+}
diff --git a/Engine/Plugins/FX/NiagaraUIRenderer/Source/NiagaraUIRenderer/NiagaraUIRenderer.Build.cs b/Engine/Plugins/FX/NiagaraUIRenderer/Source/NiagaraUIRenderer/NiagaraUIRenderer.Build.cs
new file mode 100644
--- /dev/null
+++ b/Engine/Plugins/FX/NiagaraUIRenderer/Source/NiagaraUIRenderer/NiagaraUIRenderer.Build.cs
@@ -0,0 +1,29 @@
+// Copyright Epic Games, Inc. All Rights Reserved.
+
+using UnrealBuildTool;
+
+public class NiagaraUIRenderer : ModuleRules
+{
+	public NiagaraUIRenderer(ReadOnlyTargetRules Target) : base(Target)
+	{
+		PCHUsage = ModuleRules.PCHUsageMode.UseExplicitOrSharedPCHs;
+
+		PublicDependencyModuleNames.AddRange(new string[]
+		{
+			"Core",
+			"CoreUObject",
+			"Engine",
+			"Niagara",
+			"NiagaraCore",
+			"SlateCore",
+			"Slate",
+			"UMG",
+		});
+
+		PrivateDependencyModuleNames.AddRange(new string[]
+		{
+			"RenderCore",
+			"RHI",
+		});
+	}
+}
diff --git a/Engine/Plugins/FX/NiagaraUIRenderer/Source/NiagaraUIRenderer/Private/NiagaraUIComponent.cpp b/Engine/Plugins/FX/NiagaraUIRenderer/Source/NiagaraUIRenderer/Private/NiagaraUIComponent.cpp
new file mode 100644
--- /dev/null
+++ b/Engine/Plugins/FX/NiagaraUIRenderer/Source/NiagaraUIRenderer/Private/NiagaraUIComponent.cpp
@@ -0,0 +1,79 @@
+// Copyright Epic Games, Inc. All Rights Reserved.
+
+#include "NiagaraUIComponent.h"
+#include "NiagaraUIRendererProperties.h"
+
+#include "NiagaraSystemInstanceController.h"
+#include "NiagaraSystemInstance.h"
+#include "NiagaraEmitterInstance.h"
+
+UNiagaraUIComponent::UNiagaraUIComponent()
+{
+	SetAllowScalability(false);
+	SetForceSolo(true);
+	SetAgeUpdateMode(ENiagaraAgeUpdateMode::TickDeltaTime);
+	SetCanRenderWhileSeeking(true);
+	SetMaxSimTime(0.0f);
+	SetAutoDestroy(false);
+	SetAutoActivate(false);
+	PrimaryComponentTick.bStartWithTickEnabled = false;
+}
+
+FPrimitiveSceneProxy* UNiagaraUIComponent::CreateSceneProxy()
+{
+	return nullptr;
+}
+
+const FNiagaraUIRenderData* UNiagaraUIComponent::GetRenderData() const
+{
+	return RenderData.Get();
+}
+
+void UNiagaraUIComponent::TickComponent(float DeltaTime, ELevelTick TickType, FActorComponentTickFunction* ThisTickFunction)
+{
+	Super::TickComponent(DeltaTime, TickType, ThisTickFunction);
+
+	CreateRenderData();
+}
+
+void UNiagaraUIComponent::CreateRenderData()
+{
+	FNiagaraSystemInstanceControllerPtr Controller = GetSystemInstanceController();
+	if (!Controller || !Controller->IsValid())
+	{
+		return;
+	}
+
+	FNiagaraSystemInstance* Instance = Controller->GetSystemInstance_Unsafe();
+	if (!Instance)
+	{
+		return;
+	}
+
+	TUniquePtr<FNiagaraUIRenderData> NewRenderData(new FNiagaraUIRenderData());
+	for (const FNiagaraEmitterInstanceRef& EmitterRef : Instance->GetEmitters())
+	{
+		FNiagaraEmitterInstance& EmitterInstance = *EmitterRef;
+
+		EmitterInstance.ForEachEnabledRenderer(
+			[&](const UNiagaraRendererProperties* Props)
+			{
+				const UNiagaraUIRendererProperties* UIProperties = Cast<const UNiagaraUIRendererProperties>(Props);
+				FNiagaraUIRendererRenderData* RendererRenderData = UIProperties ? UIProperties->CreateRenderData(EmitterInstance) : nullptr;
+				if (RendererRenderData)
+				{
+					check(RendererRenderData->WeakProperties.Get() != nullptr);
+
+					NewRenderData->RendererRenderDatas.Emplace(RendererRenderData);
+				}
+			}
+		);
+	}
+
+	if (NewRenderData->RendererRenderDatas.Num() > 0)
+	{
+		NewRenderData->RendererRenderDatas.StableSort([](const TUniquePtr<FNiagaraUIRendererRenderData>& DataA, const TUniquePtr<FNiagaraUIRendererRenderData>& DataB) { return DataA->SortOrder < DataB->SortOrder; });
+	}
+
+	Swap(RenderData, NewRenderData);
+}
diff --git a/Engine/Plugins/FX/NiagaraUIRenderer/Source/NiagaraUIRenderer/Private/NiagaraUIRenderContext.cpp b/Engine/Plugins/FX/NiagaraUIRenderer/Source/NiagaraUIRenderer/Private/NiagaraUIRenderContext.cpp
new file mode 100644
--- /dev/null
+++ b/Engine/Plugins/FX/NiagaraUIRenderer/Source/NiagaraUIRenderer/Private/NiagaraUIRenderContext.cpp
@@ -0,0 +1,172 @@
+// Copyright Epic Games, Inc. All Rights Reserved.
+
+#include "NiagaraUIRenderContext.h"
+
+#include "Materials/MaterialInterface.h"
+#include "Rendering/DrawElementTypes.h"
+#include "SlateMaterialBrush.h"
+
+FNiagaraUIRenderContext::FNiagaraUIRenderContext(FMaterialCacheMap& InCacheMap, const FGeometry& InAllottedGeometry, FSlateWindowElementList& InOutDrawElements)
+	: MaterialHandleCache(InCacheMap)
+	, AllottedGeometry(InAllottedGeometry)
+	, OutDrawElements(InOutDrawElements)
+{
+	const FVector2D GeometrySize = AllottedGeometry.GetLocalSize();
+	ScreenOrigin.X = float(GeometrySize.X) * 0.5f;
+	ScreenOrigin.Y = float(GeometrySize.Y) * 0.5f;
+}
+
+FNiagaraUIRendererScratchBuffers FNiagaraUIRenderContext::GetScratchBuffers(int32 NumVertices, int32 NumInidices) const
+{
+	ScratchVertices.SetNumUninitialized(NumVertices, EAllowShrinking::No);
+	ScratchIndices.SetNumUninitialized(NumInidices, EAllowShrinking::No);
+
+	return FNiagaraUIRendererScratchBuffers { ScratchVertices, ScratchIndices };
+}
+
+TConstArrayView<int32> FNiagaraUIRenderContext::CreateParticleOrderTable(int32 NumInstances, ENiagaraSortMode SortMode, FNiagaraDataSetReaderFloat<FNiagaraPosition> PositionReader, FNiagaraDataSetReaderFloat<float> SortOrderReader, int32 RendererVisibility, FNiagaraDataSetReaderInt32<int32> VisTagReader) const
+{
+	ParticleOrderTable.SetNumUninitialized(NumInstances, EAllowShrinking::No);
+
+	// Generate order table using vis tag to remove entries
+	{
+		int32 NumCulledInstances = 0;
+		for (int32 i=0; i < NumInstances; ++i)
+		{
+			if (VisTagReader.GetSafe(i, RendererVisibility) == RendererVisibility)
+			{
+				ParticleOrderTable[NumCulledInstances] = i;
+				++NumCulledInstances;
+			}
+		}
+		if (NumCulledInstances == 0)
+		{
+			return {};
+		}
+
+		NumInstances = NumCulledInstances;
+		ParticleOrderTable.SetNumUninitialized(NumCulledInstances);
+	}
+
+	// Sort the data if it's bound
+	switch (SortMode)
+	{
+		case ENiagaraSortMode::ViewDepth:
+		case ENiagaraSortMode::ViewDistance:
+			if (PositionReader.IsValid())
+			{
+				int32 DepthComponent = 0;
+				switch (ScreenPlane)
+				{
+					case ENiagaraUIScreenPlane::XY: DepthComponent = 2; break;
+					case ENiagaraUIScreenPlane::XZ: DepthComponent = 1; break;
+					case ENiagaraUIScreenPlane::YZ: DepthComponent = 0; break;
+				}
+
+				ParticleOrderTable.Sort(
+					[&](int32 IndexA, const int32 IndexB)
+					{
+						return PositionReader.Get(IndexA).Component(DepthComponent) < PositionReader.Get(IndexB).Component(DepthComponent);
+					}
+				);
+			}
+			break;
+
+		case ENiagaraSortMode::CustomAscending:
+			if (SortOrderReader.IsValid())
+			{
+				ParticleOrderTable.Sort(
+					[&](int32 IndexA, const int32 IndexB)
+					{
+						return SortOrderReader.Get(IndexA) < SortOrderReader.Get(IndexB);
+					}
+				);
+			}
+			break;
+
+		case ENiagaraSortMode::CustomDecending:
+			if (SortOrderReader.IsValid())
+			{
+				ParticleOrderTable.Sort(
+					[&](int32 IndexA, const int32 IndexB)
+					{
+						return SortOrderReader.Get(IndexA) > SortOrderReader.Get(IndexB);
+					}
+				);
+			}
+			break;
+	}
+
+	return ParticleOrderTable;
+}
+
+void FNiagaraUIRenderContext::DrawCustomVerts(const UMaterialInterface* Material) const
+{
+	DrawCustomVerts(Material, ScratchVertices, ScratchIndices);
+
+	ScratchVertices.Reset();
+	ScratchIndices.Reset();
+}
+
+void FNiagaraUIRenderContext::DrawCustomVerts(const UMaterialInterface* Material, const TArray<FSlateVertex>& Vertices, const TArray<SlateIndex>& Indices) const
+{
+	if (!Material)
+	{
+		//-TODO: Use default material??
+		return;
+	}
+
+	const FSlateResourceHandle& Handle = GetOrCreateMaterialHandle(Material);
+	if (!Handle.IsValid())
+	{
+		return;
+	}
+
+	FSlateDrawElement::MakeCustomVerts(OutDrawElements, LayerId, Handle, Vertices, Indices, nullptr, 0, 0, DrawEffect);
+}
+
+const FSlateResourceHandle& FNiagaraUIRenderContext::GetOrCreateMaterialHandle(const UMaterialInterface* Material) const
+{
+	check(IsInGameThread());
+
+	// Remove stale entries?
+	FObjectKey MaterialKey(Material);
+	if (const FSlateResourceHandle* Existing = MaterialHandleCache.Find(MaterialKey))
+	{
+		if (Existing->IsValid())
+		{
+			return *Existing;
+		}
+		MaterialHandleCache.Remove(MaterialKey);
+	}
+
+	if (!FSlateApplication::IsInitialized())
+	{
+		static FSlateResourceHandle Invalid;
+		return Invalid;
+	}
+
+	// FSlateMaterialBrush doesn't take a const UMaterialInterface
+	FSlateMaterialBrush Brush(*const_cast<UMaterialInterface*>(Material), FVector2D(1.0, 1.0));
+	FSlateResourceHandle NewHandle = FSlateApplication::Get().GetRenderer()->GetResourceHandle(Brush);
+
+	return MaterialHandleCache.Add(MaterialKey, MoveTemp(NewHandle));
+}
+
+void FNiagaraUIRenderContext::SetScreenOriginAlignment(EHorizontalAlignment HorizontalAlignment, EVerticalAlignment VerticalAlignment)
+{
+	const FVector2D GeometrySize = AllottedGeometry.GetLocalSize();
+	switch (HorizontalAlignment)
+	{
+		case HAlign_Left:	ScreenOrigin.X = 0.0f; break;
+		case HAlign_Center:	ScreenOrigin.X = float(GeometrySize.X) * 0.5f; break;
+		case HAlign_Right:	ScreenOrigin.X = float(GeometrySize.X); break;
+	}
+
+	switch (VerticalAlignment)
+	{
+		case VAlign_Top:	ScreenOrigin.Y = 0.0f; break;
+		case VAlign_Center:	ScreenOrigin.Y = float(GeometrySize.Y) * 0.5f; break;
+		case VAlign_Bottom:	ScreenOrigin.Y = float(GeometrySize.Y); break;
+	}
+}
diff --git a/Engine/Plugins/FX/NiagaraUIRenderer/Source/NiagaraUIRenderer/Private/NiagaraUIRenderer.cpp b/Engine/Plugins/FX/NiagaraUIRenderer/Source/NiagaraUIRenderer/Private/NiagaraUIRenderer.cpp
new file mode 100644
--- /dev/null
+++ b/Engine/Plugins/FX/NiagaraUIRenderer/Source/NiagaraUIRenderer/Private/NiagaraUIRenderer.cpp
@@ -0,0 +1,21 @@
+// Copyright Epic Games, Inc. All Rights Reserved.
+
+#include "Modules/ModuleManager.h"
+#include "NiagaraUISpriteRendererProperties.h"
+
+class FNiagaraUIRendererModule : public IModuleInterface
+{
+public:
+	virtual void StartupModule() override
+	{
+		// Process any UNiagaraUISpriteRendererProperties instances that tried to
+		// InitBindings() before the Niagara module was loaded.
+		UNiagaraUISpriteRendererProperties::ProcessDeferredInit();
+	}
+
+	virtual void ShutdownModule() override
+	{
+	}
+};
+
+IMPLEMENT_MODULE(FNiagaraUIRendererModule, NiagaraUIRenderer)
diff --git a/Engine/Plugins/FX/NiagaraUIRenderer/Source/NiagaraUIRenderer/Private/NiagaraUISpriteRendererProperties.cpp b/Engine/Plugins/FX/NiagaraUIRenderer/Source/NiagaraUIRenderer/Private/NiagaraUISpriteRendererProperties.cpp
new file mode 100644
--- /dev/null
+++ b/Engine/Plugins/FX/NiagaraUIRenderer/Source/NiagaraUIRenderer/Private/NiagaraUISpriteRendererProperties.cpp
@@ -0,0 +1,320 @@
+// Copyright Epic Games, Inc. All Rights Reserved.
+
+#include "NiagaraUISpriteRendererProperties.h"
+#include "NiagaraUIRenderContext.h"
+#include "NiagaraConstants.h"
+#include "NiagaraEmitterInstance.h"
+
+#include "Modules/ModuleManager.h"
+#include "Widgets/Text/STextBlock.h"
+#include "MaterialDomain.h"
+
+#define LOCTEXT_NAMESPACE "NiagaraUIRenderer"
+
+namespace NiagaraUISpriteRendererPropertiesPrivate
+{
+	static TArray<TWeakObjectPtr<UNiagaraUISpriteRendererProperties>> UISpriteRendererPropertiesToDeferredInit;
+} //namespace NiagaraUISpriteRendererPropertiesPrivate
+
+void UNiagaraUISpriteRendererProperties::ProcessDeferredInit()
+{
+	using namespace NiagaraUISpriteRendererPropertiesPrivate;
+
+	for (TWeakObjectPtr<UNiagaraUISpriteRendererProperties>& Weak : UISpriteRendererPropertiesToDeferredInit)
+	{
+		if (UNiagaraUISpriteRendererProperties* Props = Weak.Get())
+		{
+			Props->InitBindings();
+		}
+	}
+	UISpriteRendererPropertiesToDeferredInit.Empty();
+}
+
+UNiagaraUISpriteRendererProperties::UNiagaraUISpriteRendererProperties()
+{
+	AttributeBindings =
+	{
+		&PositionBinding,
+		&ColorBinding,
+		&SpriteSizeBinding,
+		&SpriteRotationBinding,
+		&SubImageIndexBinding,
+		&DynamicMaterialParameterBinding,
+		&CustomSortingBinding,
+		&RendererVisibilityTagBinding,
+	};
+};
+
+void UNiagaraUISpriteRendererProperties::PostInitProperties()
+{
+	using namespace NiagaraUISpriteRendererPropertiesPrivate;
+
+	Super::PostInitProperties();
+
+	if (!HasAnyFlags(RF_ClassDefaultObject))
+	{
+		if (!FModuleManager::Get().IsModuleLoaded("Niagara"))
+		{
+			UISpriteRendererPropertiesToDeferredInit.Add(this);
+			return;
+		}
+		InitBindings();
+	}
+}
+
+void UNiagaraUISpriteRendererProperties::PostLoad()
+{
+	Super::PostLoad();
+
+	// Ensure bindings have been set if this object was loaded from disk without them
+	if (PositionBinding.GetParamMapBindableVariable().GetName() == NAME_None)
+	{
+		InitBindings();
+	}
+}
+
+void UNiagaraUISpriteRendererProperties::InitBindings()
+{
+	if (PositionBinding.GetParamMapBindableVariable().GetName() != NAME_None)
+	{
+		// Already initialized
+		return;
+	}
+
+	PositionBinding					= FNiagaraConstants::GetAttributeDefaultBinding(SYS_PARAM_PARTICLES_POSITION);
+	ColorBinding					= FNiagaraConstants::GetAttributeDefaultBinding(SYS_PARAM_PARTICLES_COLOR);
+	SpriteSizeBinding				= FNiagaraConstants::GetAttributeDefaultBinding(SYS_PARAM_PARTICLES_SPRITE_SIZE);
+	SpriteRotationBinding			= FNiagaraConstants::GetAttributeDefaultBinding(SYS_PARAM_PARTICLES_SPRITE_ROTATION);
+	SubImageIndexBinding			= FNiagaraConstants::GetAttributeDefaultBinding(SYS_PARAM_PARTICLES_SUB_IMAGE_INDEX);
+	DynamicMaterialParameterBinding	= FNiagaraConstants::GetAttributeDefaultBinding(SYS_PARAM_PARTICLES_DYNAMIC_MATERIAL_PARAM);
+	CustomSortingBinding			= FNiagaraConstants::GetAttributeDefaultBinding(SYS_PARAM_PARTICLES_AGE);
+	RendererVisibilityTagBinding	= FNiagaraConstants::GetAttributeDefaultBinding(SYS_PARAM_PARTICLES_VISIBILITY_TAG);
+}
+
+FNiagaraRenderer* UNiagaraUISpriteRendererProperties::CreateEmitterRenderer(ERHIFeatureLevel::Type /*FeatureLevel*/, const FNiagaraEmitterInstance* /*Emitter*/, const FNiagaraSystemInstanceController& /*InController*/)
+{
+	// No 3D scene renderer — all rendering is done via Slate on the game thread.
+	return nullptr;
+}
+
+FNiagaraBoundsCalculator* UNiagaraUISpriteRendererProperties::CreateBoundsCalculator()
+{
+	// No 3D bounds — widget clips particles automatically.
+	return nullptr;
+}
+
+void UNiagaraUISpriteRendererProperties::GetUsedMaterials(const FNiagaraEmitterInstance* /*EmitterInstance*/, TArray<UMaterialInterface*>& OutMaterials) const
+{
+	if (Material)
+	{
+		OutMaterials.AddUnique(Material);
+	}
+}
+
+bool UNiagaraUISpriteRendererProperties::IsSimTargetSupported(ENiagaraSimTarget InSimTarget) const
+{
+	return InSimTarget == ENiagaraSimTarget::CPUSim;
+}
+
+#if WITH_EDITORONLY_DATA
+void UNiagaraUISpriteRendererProperties::GetRendererWidgets(const FNiagaraEmitterInstance* InEmitter, TArray<TSharedPtr<SWidget>>& OutWidgets, TSharedPtr<FAssetThumbnailPool> InThumbnailPool) const
+{
+	TArray<UMaterialInterface*> Materials;
+	GetUsedMaterials(InEmitter, Materials);
+	CreateRendererWidgetsForAssets(Materials, InThumbnailPool, OutWidgets);
+}
+
+void UNiagaraUISpriteRendererProperties::GetRendererTooltipWidgets(const FNiagaraEmitterInstance* InEmitter, TArray<TSharedPtr<SWidget>>& OutWidgets, TSharedPtr<FAssetThumbnailPool> InThumbnailPool) const
+{
+	TArray<UMaterialInterface*> Materials;
+	GetUsedMaterials(InEmitter, Materials);
+	if (Materials.Num() > 0)
+	{
+		GetRendererWidgets(InEmitter, OutWidgets, InThumbnailPool);
+	}
+	else
+	{
+		OutWidgets.Add(
+			SNew(STextBlock)
+			.Text(LOCTEXT("UISpriteRendererNoMat", "UI Sprite Renderer (No Material Set)"))
+		);
+	}
+}
+
+void UNiagaraUISpriteRendererProperties::GetRendererFeedback(const FVersionedNiagaraEmitter& InEmitter, TArray<FNiagaraRendererFeedback>& OutErrors, TArray<FNiagaraRendererFeedback>& OutWarnings, TArray<FNiagaraRendererFeedback>& OutInfo) const
+{
+	Super::GetRendererFeedback(InEmitter, OutErrors, OutWarnings, OutInfo);
+
+	if (Material == nullptr || Material->GetBaseMaterial() == nullptr)
+	{
+		OutErrors.Emplace(
+			LOCTEXT("NoMaterial", "No material is assigned nothing will render."),
+			LOCTEXT("NoMaterialSummary", "No Material Assigned.")
+		);
+	}
+	else if (Material->GetBaseMaterial()->MaterialDomain != MD_UI)
+	{
+		OutErrors.Emplace(
+			LOCTEXT("InvalidMaterial", "The current material is invalid as it is not a UI domain material."),
+			LOCTEXT("InvalidMaterialSummary", "Invalid Material Assigned.")
+		);
+	}
+}
+#endif //WITH_EDITORONLY_DATA
+
+void UNiagaraUISpriteRendererProperties::CacheFromCompiledData(const FNiagaraDataSetCompiledData* CompiledData)
+{
+	InitParticleDataSetAccessor(PositionDataSetAccessor, CompiledData, PositionBinding);
+	InitParticleDataSetAccessor(ColorDataSetAccessor, CompiledData, ColorBinding);
+	InitParticleDataSetAccessor(SpriteSizeDataSetAccessor, CompiledData, SpriteSizeBinding);
+	InitParticleDataSetAccessor(SpriteRotationDataSetAccessor, CompiledData, SpriteRotationBinding);
+	InitParticleDataSetAccessor(SubImageIndexDataSetAccessor, CompiledData, SubImageIndexBinding);
+	InitParticleDataSetAccessor(DynamicMaterialParameterDataSetAccessor, CompiledData, DynamicMaterialParameterBinding);
+	InitParticleDataSetAccessor(CustomSortingDataSetAccessor, CompiledData, CustomSortingBinding);
+	InitParticleDataSetAccessor(RendererVisibilityTagDataSetAccessor, CompiledData, RendererVisibilityTagBinding);
+}
+
+FNiagaraUIRendererRenderData* UNiagaraUISpriteRendererProperties::CreateRenderData(const FNiagaraEmitterInstance& EmitterInstance) const
+{
+	const FNiagaraDataSet& ParicleDataSet = EmitterInstance.GetParticleData();
+	FNiagaraDataBuffer* ParticleData = ParicleDataSet.GetCurrentData();
+	if (!ParticleData || ParticleData->GetNumInstances() == 0)
+	{
+		return nullptr;
+	}
+
+	UMaterial* BaseMaterial = Material ? Material->GetMaterial() : nullptr;
+	if (!BaseMaterial || BaseMaterial->MaterialDomain != MD_UI)
+	{
+		return nullptr;
+	}
+
+	FNiagaraUIRendererRenderData* RenderData = new FNiagaraUIRendererRenderData();
+	RenderData->WeakProperties	= this;
+	RenderData->DataBuffer		= ParticleData;
+	RenderData->WeakMaterial	= Material;
+	RenderData->SortOrder		= SortOrderHint;
+
+	return RenderData;
+}
+
+void UNiagaraUISpriteRendererProperties::ExecuteRender(const FNiagaraUIRenderContext& RenderContext, const FNiagaraUIRendererRenderData& RendererRenderData) const
+{
+	const FNiagaraDataBuffer* ParticleData = RendererRenderData.DataBuffer;
+
+	// Get readers
+	auto PositionReader			= PositionDataSetAccessor.GetReader(ParticleData);
+	auto ColorReader			= ColorDataSetAccessor.GetReader(ParticleData);
+	auto SpriteSizeReader		= SpriteSizeDataSetAccessor.GetReader(ParticleData);
+	auto SpriteRotationReader	= SpriteRotationDataSetAccessor.GetReader(ParticleData);
+	auto SubImageIndexReader	= SubImageIndexDataSetAccessor.GetReader(ParticleData);
+	auto DynamicParamReader		= DynamicMaterialParameterDataSetAccessor.GetReader(ParticleData);
+
+	// Default values
+	const FNiagaraPosition DefaultPosition	= FNiagaraPosition(0.0f, 0.0f, 0.0f);	//ParameterStore.GetParameterValueOrDefault(RendererProperties->PositionBinding.GetParamMapBindableVariable(), DefaultPos);
+	const FLinearColor DefaultColor			= FLinearColor::White;					//SYS_PARAM_PARTICLES_COLOR.GetValue<FNiagaraPosition>();
+	const FVector2f DefaultSpriteSize		= FVector2f(10.0f);						//SYS_PARAM_PARTICLES_SPRITE_SIZE.GetValue<FVector2f>();
+	const float DefaultSpriteRotation		= 0.0f;									//SYS_PARAM_PARTICLES_SPRITE_ROTATION.GetValue<float>();
+	const float DefaultSubImageIndex		= 0;									//SubImageIndexDataSetAccessor.GetReader(ParticleData);
+	const FVector4f DefaultDynamicParam		= FVector4f::Zero();					//SYS_PARAM_PARTICLES_DYNAMIC_MATERIAL_PARAM.GetValue<FVector4f>();
+
+	// Sub UV data
+	const float SubUVInvCols = 1.f / FMath::Max(SubImageSize.X, 1.f);
+	const float SubUVInvRows = 1.f / FMath::Max(SubImageSize.Y, 1.f);
+	const int32 SubUVNumCols = FMath::Max(FMath::RoundToInt32(SubImageSize.X), 1);
+	const int32 SubUVNumRows = FMath::Max(FMath::RoundToInt32(SubImageSize.Y), 1);
+	const int32 SubUVNumFrames = SubUVNumCols * SubUVNumRows;
+
+	// Generate particle indirection table (if needed)
+	TConstArrayView<int32> ParticleOrderTable;
+	if (SortMode != ENiagaraSortMode::None || RendererVisibilityTagDataSetAccessor.IsValid())
+	{
+		ParticleOrderTable = RenderContext.CreateParticleOrderTable(ParticleData->GetNumInstances(), SortMode, PositionReader, CustomSortingDataSetAccessor.GetReader(ParticleData), RendererVisibility, RendererVisibilityTagDataSetAccessor.GetReader(ParticleData));
+		if (ParticleOrderTable.Num() == 0)
+		{
+			return;
+		}
+	}
+
+	const uint32 NumInstances = ParticleOrderTable.Num() > 0 ? ParticleOrderTable.Num() : ParticleData->GetNumInstances();
+
+	FNiagaraUIRendererScratchBuffers ScratchBuffers = RenderContext.GetScratchBuffers(NumInstances * 4, NumInstances * 6);
+	ScratchBuffers.ScratchIndices.SetNumUninitialized(NumInstances * 6);
+	ScratchBuffers.ScratchVertices.SetNumUninitialized(NumInstances * 4);
+
+	for (uint32 iInstance = 0; iInstance < NumInstances; ++iInstance)
+	{
+		const uint32			iParticle			= ParticleOrderTable.Num() > 0 ? ParticleOrderTable[iInstance] : iInstance;
+		const FNiagaraPosition	SpritePosition		= PositionReader.GetSafe(iParticle, DefaultPosition);
+		const FLinearColor		SpriteColor			= ColorReader.GetSafe(iParticle, DefaultColor);
+		const FVector2f			SpriteSize			= RenderContext.SizeToScreen(SpriteSizeReader.GetSafe(iParticle, DefaultSpriteSize));
+		const float				SpriteRotation		= SpriteRotationReader.GetSafe(iParticle, DefaultSpriteRotation);
+		const FVector4f			SpriteDynamicParam	= DynamicParamReader.GetSafe(iParticle, DefaultDynamicParam);
+
+		const FVector2f SpriteCorners[] =
+		{
+			FVector2f(-SpriteSize.X, -SpriteSize.Y),
+			FVector2f( SpriteSize.X, -SpriteSize.Y),
+			FVector2f( SpriteSize.X,  SpriteSize.Y),
+			FVector2f(-SpriteSize.X,  SpriteSize.Y),
+		};
+
+		FVector4f SpritePackedUV(0.0f, 0.0f, 1.0f, 1.0f);
+		if (SubUVNumFrames > 1)
+		{
+			const float SpriteSubImageIndex = SubImageIndexReader.GetSafe(iParticle, DefaultSubImageIndex);
+
+			const int32 FrameA = FMath::FloorToInt32(SpriteSubImageIndex) % SubUVNumFrames;
+			const int32 ColA = FrameA % SubUVNumCols;
+			const int32 RowA = FrameA / SubUVNumCols;
+
+			SpritePackedUV.X = ColA * SubUVInvCols;
+			SpritePackedUV.Y = RowA * SubUVInvRows;
+			SpritePackedUV.Z = SpritePackedUV.X + SubUVInvCols;
+			SpritePackedUV.W = SpritePackedUV.Y + SubUVInvRows;
+		}
+
+		const FVector2f SpriteUVs[] =
+		{
+			FVector2f(SpritePackedUV.X, SpritePackedUV.Y),
+			FVector2f(SpritePackedUV.Z, SpritePackedUV.Y),
+			FVector2f(SpritePackedUV.Z, SpritePackedUV.W),
+			FVector2f(SpritePackedUV.X, SpritePackedUV.W),
+		};
+
+		const float CosRot = FMath::Cos(SpriteRotation);
+		const float SinRot = FMath::Sin(SpriteRotation);
+
+		const FVector2f Position = RenderContext.PositionToScreen(SpritePosition);
+
+		const int32 VertexOffset = iInstance * 4;
+		for ( int32 iVertex=0; iVertex < 4; ++iVertex)
+		{
+			FSlateVertex& Vertex		= ScratchBuffers.ScratchVertices[VertexOffset + iVertex];
+			Vertex.Position.X			= Position.X + ( SpriteCorners[iVertex].X * CosRot) + (SpriteCorners[iVertex].Y * SinRot);
+			Vertex.Position.Y			= Position.Y + (-SpriteCorners[iVertex].X * SinRot) + (SpriteCorners[iVertex].Y * CosRot);
+			Vertex.TexCoords[0]			= SpriteUVs[iVertex].X;
+			Vertex.TexCoords[1]			= SpriteUVs[iVertex].Y;
+			Vertex.TexCoords[2]			= SpriteDynamicParam.X;
+			Vertex.TexCoords[3]			= SpriteDynamicParam.Y;
+			Vertex.MaterialTexCoords.X	= SpriteDynamicParam.Z;
+			Vertex.MaterialTexCoords.Y	= SpriteDynamicParam.W;
+			Vertex.PixelSize[0]			= 0;
+			Vertex.PixelSize[1]			= 0;
+			Vertex.Color				= SpriteColor.ToFColor(false);		// Not sRGB
+		}
+
+		const int32 IndexOffset = iInstance * 6;
+		ScratchBuffers.ScratchIndices[IndexOffset + 0] = VertexOffset + 0;
+		ScratchBuffers.ScratchIndices[IndexOffset + 1] = VertexOffset + 1;
+		ScratchBuffers.ScratchIndices[IndexOffset + 2] = VertexOffset + 2;
+
+		ScratchBuffers.ScratchIndices[IndexOffset + 3] = VertexOffset + 0;
+		ScratchBuffers.ScratchIndices[IndexOffset + 4] = VertexOffset + 2;
+		ScratchBuffers.ScratchIndices[IndexOffset + 5] = VertexOffset + 3;
+	}
+
+	RenderContext.DrawCustomVerts(Material);
+}
+
+#undef LOCTEXT_NAMESPACE
diff --git a/Engine/Plugins/FX/NiagaraUIRenderer/Source/NiagaraUIRenderer/Private/NiagaraUIWidget.cpp b/Engine/Plugins/FX/NiagaraUIRenderer/Source/NiagaraUIRenderer/Private/NiagaraUIWidget.cpp
new file mode 100644
--- /dev/null
+++ b/Engine/Plugins/FX/NiagaraUIRenderer/Source/NiagaraUIRenderer/Private/NiagaraUIWidget.cpp
@@ -0,0 +1,117 @@
+// Copyright Epic Games, Inc. All Rights Reserved.
+
+#include "NiagaraUIWidget.h"
+#include "SNiagaraUIWidget.h"
+#include "NiagaraUIComponent.h"
+
+#include "NiagaraSystem.h"
+
+#include "Engine/World.h"
+
+UNiagaraUIWidget::UNiagaraUIWidget(const FObjectInitializer& ObjectInitializer)
+	: Super(ObjectInitializer)
+{
+}
+
+#if WITH_EDITOR
+void UNiagaraUIWidget::PostEditChangeProperty(FPropertyChangedEvent& PropertyChangedEvent)
+{
+	Super::PostEditChangeProperty(PropertyChangedEvent);
+
+	if (PropertyChangedEvent.MemberProperty)
+	{
+		const FName PropertyName = PropertyChangedEvent.MemberProperty->GetFName();
+		//NiagaraSystem
+
+		//-todo:
+		//RecreateOrInitialize();
+		//RebuildWidget();
+		//ReleaseSlateResources(true);
+		//TakeWidget();
+	}
+}
+#endif //WITH_EDITOR
+
+TSharedRef<SWidget> UNiagaraUIWidget::RebuildWidget()
+{
+	NiagaraWidget = SNew(SNiagaraUIWidget);
+
+	RecreateOrInitialize();
+
+	return NiagaraWidget.ToSharedRef();
+}
+
+void UNiagaraUIWidget::SynchronizeProperties()
+{
+	Super::SynchronizeProperties();
+
+	if (!NiagaraWidget.IsValid())
+	{
+		return;
+	}
+
+	RecreateOrInitialize();
+}
+
+void UNiagaraUIWidget::ReleaseSlateResources(bool bReleaseChildren)
+{
+	Super::ReleaseSlateResources(bReleaseChildren);
+
+	NiagaraWidget.Reset();
+	if (NiagaraComponent)
+	{
+		NiagaraComponent->DestroyComponent();
+		NiagaraComponent = nullptr;
+	}
+}
+
+#if WITH_EDITOR
+const FText UNiagaraUIWidget::GetPaletteCategory()
+{
+	return FText::FromString(TEXT("Niagara"));
+}
+#endif
+
+void UNiagaraUIWidget::SetDesiredWidgetSize(FVector2D InSize)
+{
+	DesiredWidgetSize = InSize;
+	if (NiagaraWidget.IsValid())
+	{
+		NiagaraWidget->SetDesiredSize(InSize);
+	}
+}
+
+void UNiagaraUIWidget::RecreateOrInitialize()
+{
+	UWorld* World = GetWorld();
+	if (!World || !World->PersistentLevel || World->bIsTearingDown || !NiagaraSystem)
+	{
+		if (NiagaraComponent != nullptr)
+		{
+			NiagaraComponent->DestroyComponent();
+			NiagaraComponent = nullptr;
+		}
+
+		return;
+	}
+
+	if (NiagaraComponent == nullptr)
+	{
+		NiagaraComponent = NewObject<UNiagaraUIComponent>(this);
+	}
+	NiagaraComponent->SetAsset(NiagaraSystem);
+
+	if (NiagaraWidget)
+	{
+		if (!NiagaraComponent->IsRegistered())
+		{
+			NiagaraComponent->RegisterComponentWithWorld(World);
+		}
+		NiagaraWidget->SetNiagaraComponent(NiagaraComponent);
+		NiagaraWidget->SetDesiredSize(DesiredWidgetSize);
+		NiagaraWidget->SetWorldToScreenScale(WorldToScreenScale);
+		NiagaraWidget->SetWorldToScreenPlane(WorldToScreenPlane);
+		NiagaraWidget->SetHorizontalAlignment(HorizontalAlignment);
+		NiagaraWidget->SetVerticalAlignment(VerticalAlignment);
+	}
+}
diff --git a/Engine/Plugins/FX/NiagaraUIRenderer/Source/NiagaraUIRenderer/Private/SNiagaraUIWidget.cpp b/Engine/Plugins/FX/NiagaraUIRenderer/Source/NiagaraUIRenderer/Private/SNiagaraUIWidget.cpp
new file mode 100644
--- /dev/null
+++ b/Engine/Plugins/FX/NiagaraUIRenderer/Source/NiagaraUIRenderer/Private/SNiagaraUIWidget.cpp
@@ -0,0 +1,112 @@
+// Copyright Epic Games, Inc. All Rights Reserved.
+
+#include "SNiagaraUIWidget.h"
+
+#include "NiagaraUIComponent.h"
+#include "NiagaraUITypes.h"
+#include "NiagaraUIRenderContext.h"
+#include "NiagaraUIRendererProperties.h"
+
+#include "Engine/World.h"
+#include "Framework/Application/SlateApplication.h"
+#include "Rendering/DrawElementTypes.h"
+
+void SNiagaraUIWidget::Construct(const FArguments& InArgs)
+{
+#if WITH_EDITOR
+	MaterialCompilationFinishedHandle = UMaterial::OnMaterialCompilationFinished().AddLambda(
+		[this](UMaterialInterface* ModifiedMaterial)
+		{
+			MaterialHandleCache.Empty();
+		}
+	);
+#endif
+}
+
+SNiagaraUIWidget::~SNiagaraUIWidget()
+{
+#if WITH_EDITOR
+	UMaterial::OnMaterialCompilationFinished().Remove(MaterialCompilationFinishedHandle);
+#endif
+}
+
+void SNiagaraUIWidget::Tick(const FGeometry& AllottedGeometry, const double InCurrentTime, const float InDeltaTime)
+{
+	Super::Tick(AllottedGeometry, InCurrentTime, InDeltaTime);
+
+	UNiagaraUIComponent* NiagaraComponent = WeakNiagaraComponent.Get();
+	if (NiagaraComponent)
+	{
+		NiagaraComponent->TickComponent(InDeltaTime, ELevelTick::LEVELTICK_All, nullptr);
+	}
+}
+
+int32 SNiagaraUIWidget::OnPaint(const FPaintArgs& Args, const FGeometry& AllottedGeometry, const FSlateRect& MyCullingRect, FSlateWindowElementList& OutDrawElements, int32 LayerId, const FWidgetStyle& InWidgetStyle, bool bParentEnabled) const
+{
+	UNiagaraUIComponent* NiagaraComponent = WeakNiagaraComponent.Get();
+	const FNiagaraUIRenderData* RenderData = NiagaraComponent ? NiagaraComponent->GetRenderData() : nullptr;
+	if (!RenderData)
+	{
+		return LayerId;
+	}
+
+	FNiagaraUIRenderContext RenderContext(MaterialHandleCache, AllottedGeometry, OutDrawElements);
+	RenderContext.SetLayerId(LayerId);
+	RenderContext.SetDrawEffect(bParentEnabled ? ESlateDrawEffect::None : ESlateDrawEffect::DisabledEffect);
+	RenderContext.SetScreenParameters(WorldToScreenPlane);
+	RenderContext.SetScreenOriginAlignment(HorizontalAlignment, VerticalAlignment);
+	RenderContext.SetScreenScale(WorldToScreenScale);
+
+	for (const TUniquePtr<FNiagaraUIRendererRenderData>& RendererRenderData : RenderData->RendererRenderDatas)
+	{
+		const UNiagaraUIRendererProperties* Properties = RendererRenderData->WeakProperties.Get();
+		if (Properties == nullptr)
+		{
+			continue;
+		}
+
+		Properties->ExecuteRender(RenderContext, *RendererRenderData.Get());
+	}
+
+	return LayerId;
+}
+
+FVector2D SNiagaraUIWidget::ComputeDesiredSize(float /*LayoutScaleMultiplier*/) const
+{
+	return DesiredWidgetSize;
+}
+
+void SNiagaraUIWidget::SetNiagaraComponent(UNiagaraUIComponent* InComponent)
+{
+	WeakNiagaraComponent = InComponent;
+	if (InComponent)
+	{
+		InComponent->Activate();
+	}
+}
+
+void SNiagaraUIWidget::SetDesiredSize(FVector2D InSize)
+{
+	DesiredWidgetSize = InSize;
+}
+
+void SNiagaraUIWidget::SetWorldToScreenScale(float InScale)
+{
+	WorldToScreenScale = InScale;
+}
+
+void SNiagaraUIWidget::SetWorldToScreenPlane(ENiagaraUIScreenPlane InScreenPlane)
+{
+	WorldToScreenPlane = InScreenPlane;
+}
+
+void SNiagaraUIWidget::SetHorizontalAlignment(EHorizontalAlignment InHorizontalAlignment)
+{
+	HorizontalAlignment = InHorizontalAlignment;
+}
+
+void SNiagaraUIWidget::SetVerticalAlignment(EVerticalAlignment InVerticalAlignment)
+{
+	VerticalAlignment = InVerticalAlignment;
+}
+
diff --git a/Engine/Plugins/FX/NiagaraUIRenderer/Source/NiagaraUIRenderer/Public/NiagaraUIComponent.h b/Engine/Plugins/FX/NiagaraUIRenderer/Source/NiagaraUIRenderer/Public/NiagaraUIComponent.h
new file mode 100644
--- /dev/null
+++ b/Engine/Plugins/FX/NiagaraUIRenderer/Source/NiagaraUIRenderer/Public/NiagaraUIComponent.h
@@ -0,0 +1,30 @@
+// Copyright Epic Games, Inc. All Rights Reserved.
+#pragma once
+
+#include "CoreMinimal.h"
+#include "NiagaraComponent.h"
+#include "NiagaraUITypes.h"
+#include "NiagaraUIComponent.generated.h"
+
+UCLASS(ClassGroup=FX, meta=(BlueprintSpawnableComponent, DisplayName="Niagara UI Component"))
+class UNiagaraUIComponent : public UNiagaraComponent
+{
+	GENERATED_BODY()
+
+public:
+	UNiagaraUIComponent();
+
+	// BEGIN: UNiagaraComponent interface
+	virtual FPrimitiveSceneProxy* CreateSceneProxy() override;
+	virtual void TickComponent(float DeltaTime, ELevelTick TickType, FActorComponentTickFunction* ThisTickFunction) override;
+	// END: UNiagaraComponent interface
+
+	// Get the last generated UI render data
+	const FNiagaraUIRenderData* GetRenderData() const;
+
+private:
+	void CreateRenderData();
+
+private:
+	TUniquePtr<FNiagaraUIRenderData> RenderData;
+};
diff --git a/Engine/Plugins/FX/NiagaraUIRenderer/Source/NiagaraUIRenderer/Public/NiagaraUIRenderContext.h b/Engine/Plugins/FX/NiagaraUIRenderer/Source/NiagaraUIRenderer/Public/NiagaraUIRenderContext.h
new file mode 100644
--- /dev/null
+++ b/Engine/Plugins/FX/NiagaraUIRenderer/Source/NiagaraUIRenderer/Public/NiagaraUIRenderContext.h
@@ -0,0 +1,79 @@
+// Copyright Epic Games, Inc. All Rights Reserved.
+#pragma once
+
+#include "NiagaraCommon.h"
+#include "NiagaraDataSetAccessor.h"
+#include "NiagaraGPUSortInfo.h"
+
+#include "NiagaraUITypes.h"
+
+#include "Rendering/RenderingCommon.h"
+#include "Rendering/SlateResourceHandle.h"
+
+class UMaterialInterface;
+
+struct FNiagaraUIRendererScratchBuffers
+{
+	TArray<FSlateVertex>&	ScratchVertices;
+	TArray<SlateIndex>&		ScratchIndices;
+};
+
+class FNiagaraUIRenderContext
+{
+	using FMaterialCacheMap = TMap<FObjectKey, FSlateResourceHandle>;
+
+public:
+	explicit FNiagaraUIRenderContext(FMaterialCacheMap& InCacheMap, const FGeometry& InAllottedGeometry, FSlateWindowElementList& InOutDrawElements);
+
+	// Get scratch buffers for generating vertex data into
+	FNiagaraUIRendererScratchBuffers GetScratchBuffers(int32 NumVertices, int32 NumInidices) const;
+
+	// Genrate an indirection table for particles both sorting and culling particles
+	TConstArrayView<int32> CreateParticleOrderTable(int32 NumInstances, ENiagaraSortMode SortMode, FNiagaraDataSetReaderFloat<FNiagaraPosition> PositionReader, FNiagaraDataSetReaderFloat<float> SortOrderReader, int32 RendererVisibility, FNiagaraDataSetReaderInt32<int32> VisTagReader) const;
+
+	// Draw the vertices build in the scratch buffer, will clear the buffer post draw
+	void DrawCustomVerts(const UMaterialInterface* Material) const;
+
+	// Draw the verticies provided
+	void DrawCustomVerts(const UMaterialInterface* Material, const TArray<FSlateVertex>& Vertices, const TArray<SlateIndex>& Indices) const;
+
+	// Convert Simulation Posiiton to Widget Screen
+	FVector2f PositionToScreen(const FNiagaraPosition& WorldPosition) const
+	{
+		FVector2f ScreenPosition;
+		ScreenPosition.X = (ScreenPlane == ENiagaraUIScreenPlane::YZ ? WorldPosition.Y : WorldPosition.X) *  ScreenScale + ScreenOrigin.X;
+		ScreenPosition.Y = (ScreenPlane == ENiagaraUIScreenPlane::XY ? WorldPosition.Y : WorldPosition.Z) * -ScreenScale + ScreenOrigin.Y;
+		return FVector2f(AllottedGeometry.LocalToAbsolute(FVector2D(ScreenPosition)));
+	}
+
+	FVector2f SizeToScreen(const FVector2f& WorldSize) const
+	{
+		return WorldSize * ScreenScale;
+	}
+
+	// Set functions to override default behavior
+	void SetLayerId(int32 InLayerId) { LayerId = InLayerId; }
+	void SetDrawEffect(ESlateDrawEffect InDrawEffect) { DrawEffect = InDrawEffect; }
+	void SetScreenParameters(ENiagaraUIScreenPlane InScreenPlane) { ScreenPlane = InScreenPlane; }
+	void SetScreenOrigin(const FVector2f& InScreenOrigin) { ScreenOrigin = InScreenOrigin; }
+	void SetScreenOriginAlignment(EHorizontalAlignment HorizontalAlignment, EVerticalAlignment VerticalAlignment);
+	void SetScreenScale(float InScreenScale) { ScreenScale = InScreenScale; }
+
+private:
+	const FSlateResourceHandle& GetOrCreateMaterialHandle(const UMaterialInterface* Material) const;
+
+private:
+	mutable TArray<FSlateVertex>	ScratchVertices;
+	mutable TArray<SlateIndex>		ScratchIndices;
+	mutable TArray<int32>			ParticleOrderTable;
+
+	FMaterialCacheMap&				MaterialHandleCache;
+	const FGeometry&				AllottedGeometry;
+	FSlateWindowElementList&		OutDrawElements;
+	int32							LayerId = 0;
+	ESlateDrawEffect				DrawEffect = ESlateDrawEffect::None;
+
+	ENiagaraUIScreenPlane			ScreenPlane = ENiagaraUIScreenPlane::XZ;
+	FVector2f						ScreenOrigin = FVector2f::Zero();
+	float							ScreenScale = 1.0f;
+};
diff --git a/Engine/Plugins/FX/NiagaraUIRenderer/Source/NiagaraUIRenderer/Public/NiagaraUIRendererProperties.cpp b/Engine/Plugins/FX/NiagaraUIRenderer/Source/NiagaraUIRenderer/Public/NiagaraUIRendererProperties.cpp
new file mode 100644
--- /dev/null
+++ b/Engine/Plugins/FX/NiagaraUIRenderer/Source/NiagaraUIRenderer/Public/NiagaraUIRendererProperties.cpp
@@ -0,0 +1,18 @@
+// Copyright Epic Games, Inc. All Rights Reserved.
+
+#include "NiagaraUIRendererProperties.h"
+
+FNiagaraRenderer* UNiagaraUIRendererProperties::CreateEmitterRenderer(ERHIFeatureLevel::Type /*FeatureLevel*/, const FNiagaraEmitterInstance* /*Emitter*/, const FNiagaraSystemInstanceController& /*InController*/)
+{
+	return nullptr;
+}
+
+FNiagaraBoundsCalculator* UNiagaraUIRendererProperties::CreateBoundsCalculator()
+{
+	return nullptr;
+}
+
+bool UNiagaraUIRendererProperties::IsSimTargetSupported(ENiagaraSimTarget InSimTarget) const
+{
+	return InSimTarget == ENiagaraSimTarget::CPUSim;
+}
diff --git a/Engine/Plugins/FX/NiagaraUIRenderer/Source/NiagaraUIRenderer/Public/NiagaraUIRendererProperties.h b/Engine/Plugins/FX/NiagaraUIRenderer/Source/NiagaraUIRenderer/Public/NiagaraUIRendererProperties.h
new file mode 100644
--- /dev/null
+++ b/Engine/Plugins/FX/NiagaraUIRenderer/Source/NiagaraUIRenderer/Public/NiagaraUIRendererProperties.h
@@ -0,0 +1,27 @@
+// Copyright Epic Games, Inc. All Rights Reserved.
+#pragma once
+
+#include "NiagaraRendererProperties.h"
+#include "NiagaraCommon.h"
+#include "NiagaraGPUSortInfo.h"
+#include "NiagaraSystemInstanceController.h"
+#include "NiagaraUITypes.h"
+#include "NiagaraUIRendererProperties.generated.h"
+
+class FNiagaraUIRenderContext;
+
+UCLASS(ABSTRACT, MinimalAPI)
+class UNiagaraUIRendererProperties : public UNiagaraRendererProperties
+{
+	GENERATED_BODY()
+
+public:
+	//~ BEGIN: UNiagaraRendererProperties interface
+	virtual FNiagaraRenderer* CreateEmitterRenderer(ERHIFeatureLevel::Type /*FeatureLevel*/, const FNiagaraEmitterInstance* /*Emitter*/, const FNiagaraSystemInstanceController& /*InController*/) override;
+	virtual FNiagaraBoundsCalculator* CreateBoundsCalculator() override;
+	virtual bool IsSimTargetSupported(ENiagaraSimTarget InSimTarget) const override;
+	//~ END: UNiagaraRendererProperties interface
+
+	virtual FNiagaraUIRendererRenderData* CreateRenderData(const FNiagaraEmitterInstance&) const { return nullptr; }
+	virtual void ExecuteRender(const FNiagaraUIRenderContext& RenderContext, const FNiagaraUIRendererRenderData& RendererRenderData) const {}
+};
diff --git a/Engine/Plugins/FX/NiagaraUIRenderer/Source/NiagaraUIRenderer/Public/NiagaraUISpriteRendererProperties.h b/Engine/Plugins/FX/NiagaraUIRenderer/Source/NiagaraUIRenderer/Public/NiagaraUISpriteRendererProperties.h
new file mode 100644
--- /dev/null
+++ b/Engine/Plugins/FX/NiagaraUIRenderer/Source/NiagaraUIRenderer/Public/NiagaraUISpriteRendererProperties.h
@@ -0,0 +1,91 @@
+// Copyright Epic Games, Inc. All Rights Reserved.
+#pragma once
+
+#include "NiagaraUIRendererProperties.h"
+#include "NiagaraUISpriteRendererProperties.generated.h"
+
+// Niagara Renderer for sprites in the UI (Slate / UMG)
+UCLASS(EditInlineNew, meta=(DisplayName="UI Sprite Renderer", PrioritizeCategories="Material Bindings SubUV Sorting"))
+class UNiagaraUISpriteRendererProperties : public UNiagaraUIRendererProperties
+{
+	GENERATED_BODY()
+
+public:
+	UNiagaraUISpriteRendererProperties();
+
+	//~ BEGIN: UNiagaraRendererProperties interface
+	virtual FNiagaraRenderer* CreateEmitterRenderer(ERHIFeatureLevel::Type FeatureLevel, const FNiagaraEmitterInstance* Emitter, const FNiagaraSystemInstanceController& InController) override;
+	virtual class FNiagaraBoundsCalculator* CreateBoundsCalculator() override;
+	virtual void GetUsedMaterials(const FNiagaraEmitterInstance* EmitterInstance, TArray<UMaterialInterface*>& OutMaterials) const override;
+	virtual bool IsSimTargetSupported(ENiagaraSimTarget InSimTarget) const override;
+	virtual void PostInitProperties() override;
+	virtual void PostLoad() override;
+#if WITH_EDITORONLY_DATA
+	virtual void GetRendererWidgets(const FNiagaraEmitterInstance* InEmitter, TArray<TSharedPtr<SWidget>>& OutWidgets, TSharedPtr<FAssetThumbnailPool> InThumbnailPool) const override;
+	virtual void GetRendererTooltipWidgets(const FNiagaraEmitterInstance* InEmitter, TArray<TSharedPtr<SWidget>>& OutWidgets, TSharedPtr<FAssetThumbnailPool> InThumbnailPool) const override;
+	virtual void GetRendererFeedback(const FVersionedNiagaraEmitter& InEmitter, TArray<FNiagaraRendererFeedback>& OutErrors, TArray<FNiagaraRendererFeedback>& OutWarnings, TArray<FNiagaraRendererFeedback>& OutInfo) const override;
+#endif
+	virtual void CacheFromCompiledData(const FNiagaraDataSetCompiledData* CompiledData) override;
+	//~ END: UNiagaraRendererProperties interface
+
+	//~ BEGIN: UNiagaraUIRendererProperties interface
+	virtual FNiagaraUIRendererRenderData* CreateRenderData(const FNiagaraEmitterInstance& EmitterInstance) const override;
+	virtual void ExecuteRender(const FNiagaraUIRenderContext& RenderContext, const FNiagaraUIRendererRenderData& RendererRenderData) const override;
+	//~ END: UNiagaraUIRendererProperties interface
+
+	void InitBindings();
+
+	static void ProcessDeferredInit();
+
+public:
+	UPROPERTY(EditAnywhere, Category="Material")//, meta=(RequiredAssetDataTags="MaterialDomain=UI"))
+	TObjectPtr<UMaterialInterface> Material;
+
+	UPROPERTY(EditAnywhere, Category="SubUV", meta=(ClampMin="1"))
+	FVector2f SubImageSize = FVector2f(1.0f, 1.0f);
+
+	//UPROPERTY(EditAnywhere, Category="SubUV", meta=(DisplayName="Sub UV Blending Enabled"))
+	//bool bSubImageBlend = false;
+
+	UPROPERTY(EditAnywhere, Category="Sorting And Visibility")
+	ENiagaraSortMode SortMode = ENiagaraSortMode::None;
+
+	/** If a render visibility tag is present, particles whose tag matches this value will be visible in this renderer. */
+	UPROPERTY(EditAnywhere, Category = "Sorting And Visibility")
+	uint32 RendererVisibility = 0;
+
+	UPROPERTY(EditAnywhere, Category="Bindings")
+	FNiagaraVariableAttributeBinding PositionBinding;
+
+	UPROPERTY(EditAnywhere, Category="Bindings")
+	FNiagaraVariableAttributeBinding ColorBinding;
+
+	UPROPERTY(EditAnywhere, Category="Bindings")
+	FNiagaraVariableAttributeBinding SpriteSizeBinding;
+
+	UPROPERTY(EditAnywhere, Category="Bindings")
+	FNiagaraVariableAttributeBinding SpriteRotationBinding;
+
+	UPROPERTY(EditAnywhere, Category="Bindings")
+	FNiagaraVariableAttributeBinding SubImageIndexBinding;
+
+	// Packed into TexCoord[1] / TexCoord[2]
+	UPROPERTY(EditAnywhere, Category="Bindings")
+	FNiagaraVariableAttributeBinding DynamicMaterialParameterBinding;
+
+	UPROPERTY(EditAnywhere, Category = "Bindings")
+	FNiagaraVariableAttributeBinding CustomSortingBinding;
+
+	UPROPERTY(EditAnywhere, Category = "Bindings")
+	FNiagaraVariableAttributeBinding RendererVisibilityTagBinding;
+
+private:
+	FNiagaraDataSetAccessor<FNiagaraPosition>	PositionDataSetAccessor;
+	FNiagaraDataSetAccessor<FLinearColor>		ColorDataSetAccessor;
+	FNiagaraDataSetAccessor<FVector2f>			SpriteSizeDataSetAccessor;
+	FNiagaraDataSetAccessor<float>				SpriteRotationDataSetAccessor;
+	FNiagaraDataSetAccessor<float>				SubImageIndexDataSetAccessor;
+	FNiagaraDataSetAccessor<FVector4f>			DynamicMaterialParameterDataSetAccessor;
+	FNiagaraDataSetAccessor<float>				CustomSortingDataSetAccessor;
+	FNiagaraDataSetAccessor<int32>				RendererVisibilityTagDataSetAccessor;
+};
diff --git a/Engine/Plugins/FX/NiagaraUIRenderer/Source/NiagaraUIRenderer/Public/NiagaraUITypes.h b/Engine/Plugins/FX/NiagaraUIRenderer/Source/NiagaraUIRenderer/Public/NiagaraUITypes.h
new file mode 100644
--- /dev/null
+++ b/Engine/Plugins/FX/NiagaraUIRenderer/Source/NiagaraUIRenderer/Public/NiagaraUITypes.h
@@ -0,0 +1,34 @@
+// Copyright Epic Games, Inc. All Rights Reserved.
+#pragma once
+
+#include "NiagaraCommon.h"
+#include "NiagaraUITypes.generated.h"
+
+class UMaterialInterface;
+class UNiagaraUIRendererProperties;
+
+/** Which world-space plane to project Niagara particle positions onto. */
+UENUM(BlueprintType)
+enum class ENiagaraUIScreenPlane : uint8
+{
+	/** Project X (right) and Y (up) — typical for effects simulated in the XY plane. */
+	XY,
+	/** Project X (right) and Z (up). */
+	XZ,
+	/** Project Y (right) and Z (up). */
+	YZ,
+};
+
+struct FNiagaraUIRendererRenderData
+{
+	TWeakObjectPtr<const UNiagaraUIRendererProperties>	WeakProperties;
+	TWeakObjectPtr<UMaterialInterface>					WeakMaterial;
+	FNiagaraDataBufferRef								DataBuffer;
+	int32												SortOrder = 0;
+};
+
+struct FNiagaraUIRenderData
+{
+	TArray<TUniquePtr<FNiagaraUIRendererRenderData>> RendererRenderDatas;
+};
+
diff --git a/Engine/Plugins/FX/NiagaraUIRenderer/Source/NiagaraUIRenderer/Public/NiagaraUIWidget.h b/Engine/Plugins/FX/NiagaraUIRenderer/Source/NiagaraUIRenderer/Public/NiagaraUIWidget.h
new file mode 100644
--- /dev/null
+++ b/Engine/Plugins/FX/NiagaraUIRenderer/Source/NiagaraUIRenderer/Public/NiagaraUIWidget.h
@@ -0,0 +1,72 @@
+// Copyright Epic Games, Inc. All Rights Reserved.
+#pragma once
+
+#include "CoreMinimal.h"
+#include "Components/Widget.h"
+#include "NiagaraUITypes.h"
+#include "NiagaraUIWidget.generated.h"
+
+class UNiagaraUIComponent;
+class SNiagaraUIWidget;
+class UNiagaraSystem;
+
+UCLASS(meta=(DisplayName="Niagara UI Widget"))
+class UNiagaraUIWidget : public UWidget
+{
+	GENERATED_BODY()
+
+public:
+	UNiagaraUIWidget(const FObjectInitializer& ObjectInitializer);
+
+	// BEGIN: UObject interface
+#if WITH_EDITOR
+	virtual void PostEditChangeProperty(FPropertyChangedEvent& PropertyChangedEvent) override;
+#endif //WITH_EDITOR
+	// END: UObject interface
+
+	// BEGIN: UWidget interface
+	virtual TSharedRef<SWidget> RebuildWidget() override;
+	virtual void SynchronizeProperties() override;
+	virtual void ReleaseSlateResources(bool bReleaseChildren) override;
+#if WITH_EDITOR
+	virtual const FText GetPaletteCategory() override;
+#endif
+	// END: UWidget interface
+
+	UFUNCTION(BlueprintCallable, Category = "Niagara")
+	void SetDesiredWidgetSize(FVector2D InSize);
+
+private:
+	void RecreateOrInitialize();
+
+protected:
+	// Layout for widget, does not impact simulation coordinate translation
+	UPROPERTY(EditAnywhere, BlueprintReadWrite, Category = "UI Rendering")
+	FVector2D DesiredWidgetSize = FVector2D(100.0, 100.0);
+
+	// Scale to translate Niagara units to UI units
+	UPROPERTY(EditAnywhere, BlueprintReadWrite, Category = "UI Rendering")
+	float WorldToScreenScale = 1.0f;
+
+	// Axis to translate Niagara to UI
+	UPROPERTY(EditAnywhere, BlueprintReadWrite, Category = "UI Rendering")
+	ENiagaraUIScreenPlane WorldToScreenPlane = ENiagaraUIScreenPlane::XY;
+
+	// Horizontal alignment that represents the world origin
+	UPROPERTY(EditAnywhere, BlueprintReadWrite, Category = "UI Rendering")
+	TEnumAsByte<EHorizontalAlignment> HorizontalAlignment = HAlign_Center;
+
+	// Vertical alignment that represents the world origin
+	UPROPERTY(EditAnywhere, BlueprintReadWrite, Category = "UI Rendering")
+	TEnumAsByte<EVerticalAlignment> VerticalAlignment = VAlign_Center;
+
+	// System simulation to use
+	UPROPERTY(EditAnywhere, BlueprintReadWrite, Category = "Niagara")
+	TObjectPtr<UNiagaraSystem> NiagaraSystem;
+
+	// Active component, used for rendering
+	UPROPERTY(Transient)
+	TObjectPtr<UNiagaraUIComponent> NiagaraComponent;
+
+	TSharedPtr<SNiagaraUIWidget> NiagaraWidget;
+};
diff --git a/Engine/Plugins/FX/NiagaraUIRenderer/Source/NiagaraUIRenderer/Public/SNiagaraUIWidget.h b/Engine/Plugins/FX/NiagaraUIRenderer/Source/NiagaraUIRenderer/Public/SNiagaraUIWidget.h
new file mode 100644
--- /dev/null
+++ b/Engine/Plugins/FX/NiagaraUIRenderer/Source/NiagaraUIRenderer/Public/SNiagaraUIWidget.h
@@ -0,0 +1,54 @@
+// Copyright Epic Games, Inc. All Rights Reserved.
+#pragma once
+
+#include "CoreMinimal.h"
+#include "Widgets/SLeafWidget.h"
+#include "Rendering/SlateResourceHandle.h"
+#include "NiagaraUITypes.h"
+
+class UNiagaraUIComponent;
+class UMaterialInterface;
+
+class SNiagaraUIWidget : public SLeafWidget
+{
+	using Super = SLeafWidget;
+
+public:
+	SLATE_BEGIN_ARGS(SNiagaraUIWidget) {}
+	SLATE_END_ARGS()
+
+	void Construct(const FArguments& InArgs);
+	virtual ~SNiagaraUIWidget();
+
+	// BEGIN: SWidget interface
+	virtual void Tick(const FGeometry& AllottedGeometry, const double InCurrentTime, const float InDeltaTime) override;
+	virtual int32 OnPaint(const FPaintArgs& Args, const FGeometry& AllottedGeometry, const FSlateRect& MyCullingRect, FSlateWindowElementList& OutDrawElements, int32 LayerId, const FWidgetStyle& InWidgetStyle, bool bParentEnabled) const override;
+	virtual FVector2D ComputeDesiredSize(float LayoutScaleMultiplier) const override;
+	virtual bool ComputeVolatility() const override { return true; }
+	// END: SWidget interface
+
+	/** Switch which component this widget reads from. */
+	void SetNiagaraComponent(UNiagaraUIComponent* InComponent);
+	void SetDesiredSize(FVector2D InSize);
+	void SetWorldToScreenScale(float InScale);
+	void SetWorldToScreenPlane(ENiagaraUIScreenPlane InScreenPlane);
+	void SetHorizontalAlignment(EHorizontalAlignment InHorizontalAlignment);
+	void SetVerticalAlignment(EVerticalAlignment InVerticalAlignment);
+
+private:
+	TWeakObjectPtr<UNiagaraUIComponent> WeakNiagaraComponent;
+
+	FVector2D DesiredWidgetSize = FVector2D(100.0, 100.0);
+
+	float					WorldToScreenScale = 1.0f;
+	ENiagaraUIScreenPlane	WorldToScreenPlane = ENiagaraUIScreenPlane::XY;
+	EHorizontalAlignment	HorizontalAlignment = HAlign_Center;
+	EVerticalAlignment		VerticalAlignment = VAlign_Center;
+
+	// Cached map from material to slate resource
+	mutable TMap<FObjectKey, FSlateResourceHandle> MaterialHandleCache;
+
+#if WITH_EDITOR
+	FDelegateHandle MaterialCompilationFinishedHandle;
+#endif
+};
diff --git a/Engine/Plugins/FX/NiagaraUIRenderer/Source/NiagaraUIRendererEditor/NiagaraUIRendererEditor.Build.cs b/Engine/Plugins/FX/NiagaraUIRenderer/Source/NiagaraUIRendererEditor/NiagaraUIRendererEditor.Build.cs
new file mode 100644
--- /dev/null
+++ b/Engine/Plugins/FX/NiagaraUIRenderer/Source/NiagaraUIRendererEditor/NiagaraUIRendererEditor.Build.cs
@@ -0,0 +1,28 @@
+// Copyright Epic Games, Inc. All Rights Reserved.
+
+using UnrealBuildTool;
+
+public class NiagaraUIRendererEditor : ModuleRules
+{
+	public NiagaraUIRendererEditor(ReadOnlyTargetRules Target) : base(Target)
+	{
+		PCHUsage = ModuleRules.PCHUsageMode.UseExplicitOrSharedPCHs;
+
+		PublicDependencyModuleNames.AddRange(new string[]
+		{
+			"Core",
+			"CoreUObject",
+			"Engine",
+			"NiagaraUIRenderer",
+		});
+
+		PrivateDependencyModuleNames.AddRange(new string[]
+		{
+			"NiagaraEditor",
+			"UnrealEd",
+			"PropertyEditor",
+			"SlateCore",
+			"Slate",
+		});
+	}
+}
diff --git a/Engine/Plugins/FX/NiagaraUIRenderer/Source/NiagaraUIRendererEditor/Private/NiagaraUIRendererEditor.cpp b/Engine/Plugins/FX/NiagaraUIRenderer/Source/NiagaraUIRendererEditor/Private/NiagaraUIRendererEditor.cpp
new file mode 100644
--- /dev/null
+++ b/Engine/Plugins/FX/NiagaraUIRenderer/Source/NiagaraUIRendererEditor/Private/NiagaraUIRendererEditor.cpp
@@ -0,0 +1,27 @@
+// Copyright Epic Games, Inc. All Rights Reserved.
+
+#include "NiagaraUISpriteRendererProperties.h"
+#include "NiagaraEditorModule.h"
+#include "Modules/ModuleManager.h"
+
+class FNiagaraUIRendererEditorModule : public IModuleInterface
+{
+public:
+	virtual void StartupModule() override
+	{
+		FNiagaraEditorModule& NiagaraEditor = FModuleManager::LoadModuleChecked<FNiagaraEditorModule>("NiagaraEditor");
+
+		NiagaraEditor.RegisterRendererCreationInfo(FNiagaraRendererCreationInfo(
+			UNiagaraUISpriteRendererProperties::StaticClass()->GetDisplayNameText(),
+			FText::FromString(UNiagaraUISpriteRendererProperties::StaticClass()->GetDescription()),
+			UNiagaraUISpriteRendererProperties::StaticClass()->GetClassPathName(),
+			FNiagaraRendererCreationInfo::FRendererFactory::CreateLambda([](UObject* OuterEmitter)
+			{
+				return NewObject<UNiagaraUISpriteRendererProperties>(OuterEmitter, NAME_None, RF_Transactional);
+			})));
+	}
+
+	virtual void ShutdownModule() override {}
+};
+
+IMPLEMENT_MODULE(FNiagaraUIRendererEditorModule, NiagaraUIRendererEditor)
