- Add debug drawing to UI renderer
- This makes it possible to see what is happening inside the system editor

#rnx

#tests Added UI Renderer and used the eye to draw the debug data
#Synced-CL 52445895

diff --git a/Engine/Plugins/FX/NiagaraUIRenderer/Source/NiagaraUIRenderer/Private/NiagaraUIRenderContext.cpp b/Engine/Plugins/FX/NiagaraUIRenderer/Source/NiagaraUIRenderer/Private/NiagaraUIRenderContext.cpp
--- a/Engine/Plugins/FX/NiagaraUIRenderer/Source/NiagaraUIRenderer/Private/NiagaraUIRenderContext.cpp
+++ b/Engine/Plugins/FX/NiagaraUIRenderer/Source/NiagaraUIRenderer/Private/NiagaraUIRenderContext.cpp
@@ -6,10 +6,9 @@
 #include "Rendering/DrawElementTypes.h"
 #include "SlateMaterialBrush.h"
 
-FNiagaraUIRenderContext::FNiagaraUIRenderContext(FMaterialCacheMap& InCacheMap, const FGeometry& InAllottedGeometry, FSlateWindowElementList& InOutDrawElements)
+FNiagaraUIRenderContext::FNiagaraUIRenderContext(FMaterialCacheMap& InCacheMap, const FGeometry& InAllottedGeometry)
 	: MaterialHandleCache(InCacheMap)
 	, AllottedGeometry(InAllottedGeometry)
-	, OutDrawElements(InOutDrawElements)
 {
 	const FVector2D GeometrySize = AllottedGeometry.GetLocalSize();
 	ScreenOrigin.X = float(GeometrySize.X) * 0.5f;
@@ -108,23 +107,6 @@ void FNiagaraUIRenderContext::DrawCustomVerts(const UMaterialInterface* Material
 	ScratchIndices.Reset();
 }
 
-void FNiagaraUIRenderContext::DrawCustomVerts(const UMaterialInterface* Material, const TArray<FSlateVertex>& Vertices, const TArray<SlateIndex>& Indices) const
-{
-	if (!Material)
-	{
-		//-TODO: Use default material??
-		return;
-	}
-
-	const FSlateResourceHandle& Handle = GetOrCreateMaterialHandle(Material);
-	if (!Handle.IsValid())
-	{
-		return;
-	}
-
-	FSlateDrawElement::MakeCustomVerts(OutDrawElements, LayerId, Handle, Vertices, Indices, nullptr, 0, 0, DrawEffect);
-}
-
 const FSlateResourceHandle& FNiagaraUIRenderContext::GetOrCreateMaterialHandle(const UMaterialInterface* Material) const
 {
 	check(IsInGameThread());
diff --git a/Engine/Plugins/FX/NiagaraUIRenderer/Source/NiagaraUIRenderer/Private/NiagaraUISlateRenderContext.cpp b/Engine/Plugins/FX/NiagaraUIRenderer/Source/NiagaraUIRenderer/Private/NiagaraUISlateRenderContext.cpp
new file mode 100644
--- /dev/null
+++ b/Engine/Plugins/FX/NiagaraUIRenderer/Source/NiagaraUIRenderer/Private/NiagaraUISlateRenderContext.cpp
@@ -0,0 +1,29 @@
+// Copyright Epic Games, Inc. All Rights Reserved.
+
+#include "NiagaraUISlateRenderContext.h"
+
+#include "Rendering/DrawElementTypes.h"
+
+FNiagaraUISlateRenderContext::FNiagaraUISlateRenderContext(FMaterialCacheMap& InCacheMap, const FGeometry& InAllottedGeometry, FSlateWindowElementList& InOutDrawElements)
+	: FNiagaraUIRenderContext(InCacheMap, InAllottedGeometry)
+	, OutDrawElements(InOutDrawElements)
+{
+}
+
+void FNiagaraUISlateRenderContext::DrawCustomVerts(const UMaterialInterface* Material, const TArray<FSlateVertex>& Vertices, const TArray<SlateIndex>& Indices) const
+{
+	if (!Material)
+	{
+		//-TODO: Use default material??
+		return;
+	}
+
+	const FSlateResourceHandle& Handle = GetOrCreateMaterialHandle(Material);
+	if (!Handle.IsValid())
+	{
+		return;
+	}
+
+	FSlateDrawElement::MakeCustomVerts(OutDrawElements, LayerId, Handle, Vertices, Indices, nullptr, 0, 0, DrawEffect);
+}
+
diff --git a/Engine/Plugins/FX/NiagaraUIRenderer/Source/NiagaraUIRenderer/Private/NiagaraUISlateRenderContext.h b/Engine/Plugins/FX/NiagaraUIRenderer/Source/NiagaraUIRenderer/Private/NiagaraUISlateRenderContext.h
new file mode 100644
--- /dev/null
+++ b/Engine/Plugins/FX/NiagaraUIRenderer/Source/NiagaraUIRenderer/Private/NiagaraUISlateRenderContext.h
@@ -0,0 +1,16 @@
+// Copyright Epic Games, Inc. All Rights Reserved.
+#pragma once
+
+#include "NiagaraUIRenderContext.h"
+
+class FNiagaraUISlateRenderContext : public FNiagaraUIRenderContext
+{
+public:
+	explicit FNiagaraUISlateRenderContext(FMaterialCacheMap& InCacheMap, const FGeometry& InAllottedGeometry, FSlateWindowElementList& InOutDrawElements);
+	virtual ~FNiagaraUISlateRenderContext() = default;
+
+	virtual void DrawCustomVerts(const UMaterialInterface* Material, const TArray<FSlateVertex>& Vertices, const TArray<SlateIndex>& Indices) const override;
+
+private:
+	FSlateWindowElementList& OutDrawElements;
+};
diff --git a/Engine/Plugins/FX/NiagaraUIRenderer/Source/NiagaraUIRenderer/Private/NiagaraUISpriteRendererProperties.cpp b/Engine/Plugins/FX/NiagaraUIRenderer/Source/NiagaraUIRenderer/Private/NiagaraUISpriteRendererProperties.cpp
--- a/Engine/Plugins/FX/NiagaraUIRenderer/Source/NiagaraUIRenderer/Private/NiagaraUISpriteRendererProperties.cpp
+++ b/Engine/Plugins/FX/NiagaraUIRenderer/Source/NiagaraUIRenderer/Private/NiagaraUISpriteRendererProperties.cpp
@@ -91,18 +91,6 @@ void UNiagaraUISpriteRendererProperties::InitBindings()
 	RendererVisibilityTagBinding	= FNiagaraConstants::GetAttributeDefaultBinding(SYS_PARAM_PARTICLES_VISIBILITY_TAG);
 }
 
-FNiagaraRenderer* UNiagaraUISpriteRendererProperties::CreateEmitterRenderer(ERHIFeatureLevel::Type /*FeatureLevel*/, const FNiagaraEmitterInstance* /*Emitter*/, const FNiagaraSystemInstanceController& /*InController*/)
-{
-	// No 3D scene renderer — all rendering is done via Slate on the game thread.
-	return nullptr;
-}
-
-FNiagaraBoundsCalculator* UNiagaraUISpriteRendererProperties::CreateBoundsCalculator()
-{
-	// No 3D bounds — widget clips particles automatically.
-	return nullptr;
-}
-
 void UNiagaraUISpriteRendererProperties::GetUsedMaterials(const FNiagaraEmitterInstance* /*EmitterInstance*/, TArray<UMaterialInterface*>& OutMaterials) const
 {
 	if (Material)
@@ -111,11 +99,6 @@ void UNiagaraUISpriteRendererProperties::GetUsedMaterials(const FNiagaraEmitterI
 	}
 }
 
-bool UNiagaraUISpriteRendererProperties::IsSimTargetSupported(ENiagaraSimTarget InSimTarget) const
-{
-	return InSimTarget == ENiagaraSimTarget::CPUSim;
-}
-
 #if WITH_EDITORONLY_DATA
 void UNiagaraUISpriteRendererProperties::GetRendererWidgets(const FNiagaraEmitterInstance* InEmitter, TArray<TSharedPtr<SWidget>>& OutWidgets, TSharedPtr<FAssetThumbnailPool> InThumbnailPool) const
 {
diff --git a/Engine/Plugins/FX/NiagaraUIRenderer/Source/NiagaraUIRenderer/Private/SNiagaraUIWidget.cpp b/Engine/Plugins/FX/NiagaraUIRenderer/Source/NiagaraUIRenderer/Private/SNiagaraUIWidget.cpp
--- a/Engine/Plugins/FX/NiagaraUIRenderer/Source/NiagaraUIRenderer/Private/SNiagaraUIWidget.cpp
+++ b/Engine/Plugins/FX/NiagaraUIRenderer/Source/NiagaraUIRenderer/Private/SNiagaraUIWidget.cpp
@@ -4,8 +4,8 @@
 
 #include "NiagaraUIComponent.h"
 #include "NiagaraUITypes.h"
-#include "NiagaraUIRenderContext.h"
 #include "NiagaraUIRendererProperties.h"
+#include "NiagaraUISlateRenderContext.h"
 
 #include "Engine/World.h"
 #include "Framework/Application/SlateApplication.h"
@@ -50,7 +50,7 @@ int32 SNiagaraUIWidget::OnPaint(const FPaintArgs& Args, const FGeometry& Allotte
 		return LayerId;
 	}
 
-	FNiagaraUIRenderContext RenderContext(MaterialHandleCache, AllottedGeometry, OutDrawElements);
+	FNiagaraUISlateRenderContext RenderContext(MaterialHandleCache, AllottedGeometry, OutDrawElements);
 	RenderContext.SetLayerId(LayerId);
 	RenderContext.SetDrawEffect(bParentEnabled ? ESlateDrawEffect::None : ESlateDrawEffect::DisabledEffect);
 	RenderContext.SetScreenParameters(WorldToScreenPlane);
diff --git a/Engine/Plugins/FX/NiagaraUIRenderer/Source/NiagaraUIRenderer/Public/NiagaraUIDebugRenderer.cpp b/Engine/Plugins/FX/NiagaraUIRenderer/Source/NiagaraUIRenderer/Public/NiagaraUIDebugRenderer.cpp
new file mode 100644
--- /dev/null
+++ b/Engine/Plugins/FX/NiagaraUIRenderer/Source/NiagaraUIRenderer/Public/NiagaraUIDebugRenderer.cpp
@@ -0,0 +1,101 @@
+// Copyright Epic Games, Inc. All Rights Reserved.
+
+#include "NiagaraUIDebugRenderer.h"
+
+#include "NiagaraUIRendererProperties.h"
+#include "NiagaraUIRenderContext.h"
+
+#include "Async/Async.h"
+#include "Components/LineBatchComponent.h"
+#include "Engine/World.h"
+
+#if WITH_NIAGARA_RENDERER_DEBUGDRAW 
+
+namespace NiagaraUIDebugRendererPrivate
+{
+	static const FGeometry DebugGeometry(FVector2f(0.0f, 0.0f), FVector2f(-64.0f, -64.0f), FVector2f(128.0f, 128.0f), 1.0f);
+}
+
+class FNiagaraUIDebugRenderContext : public FNiagaraUIRenderContext
+{
+public:
+	FNiagaraUIDebugRenderContext(FMaterialCacheMap& TempCacheMap, ULineBatchComponent& InLineBatcher)
+		: FNiagaraUIRenderContext(TempCacheMap, NiagaraUIDebugRendererPrivate::DebugGeometry)
+		, LineBatcher(InLineBatcher)
+	{
+	}
+
+	virtual void DrawCustomVerts(const UMaterialInterface* Material, const TArray<FSlateVertex>& Vertices, const TArray<SlateIndex>& Indices) const override
+	{
+		TArray<FBatchedLine> Lines;
+
+		const int32 NumTriangles = Indices.Num() / 3;
+		for (int iTriangle = 0; iTriangle < NumTriangles; ++iTriangle)
+		{
+			const FSlateVertex& Vertex0 = Vertices[Indices[iTriangle * 3 + 0]];
+			const FSlateVertex& Vertex1 = Vertices[Indices[iTriangle * 3 + 1]];
+			const FSlateVertex& Vertex2 = Vertices[Indices[iTriangle * 3 + 2]];
+
+			const FVector Positions[] =
+			{
+				FVector(Vertex0.Position.X, 0.0f, -Vertex0.Position.Y),
+				FVector(Vertex1.Position.X, 0.0f, -Vertex1.Position.Y),
+				FVector(Vertex2.Position.X, 0.0f, -Vertex2.Position.Y),
+			};
+
+			Lines.Emplace(Positions[0], Positions[1], FLinearColor::White, 0.0f, 1.0f, 0);
+			Lines.Emplace(Positions[1], Positions[2], FLinearColor::White, 0.0f, 1.0f, 0);
+			Lines.Emplace(Positions[2], Positions[0], FLinearColor::White, 0.0f, 1.0f, 0);
+		}
+
+		AsyncTask(
+			ENamedThreads::GameThread,
+			[WeakLineBatcher=MakeWeakObjectPtr(&LineBatcher), TaskLines=MoveTemp(Lines)]() mutable
+			{
+				if ( ULineBatchComponent* LineBatcher = WeakLineBatcher.Get() )
+				{
+					LineBatcher->DrawLines(TaskLines);
+				}
+			}
+		);
+	}
+
+private:
+	ULineBatchComponent&	LineBatcher;
+};
+
+
+FNiagaraUIDebugRenderer::FNiagaraUIDebugRenderer(ERHIFeatureLevel::Type FeatureLevel, const UNiagaraRendererProperties* InProps, const FNiagaraEmitterInstance* Emitter)
+	: FNiagaraRenderer(FeatureLevel, InProps, Emitter)
+{
+}
+
+FNiagaraDynamicDataBase* FNiagaraUIDebugRenderer::GenerateDynamicData(const FNiagaraSceneProxy* Proxy, const UNiagaraRendererProperties* InProperties, const FNiagaraEmitterInstance* Emitter) const
+{
+	const UNiagaraUIRendererProperties* RendererProperties = CastChecked<const UNiagaraUIRendererProperties>(InProperties);
+	if (!RendererProperties->IsDebugDrawEnabled())
+	{
+		return nullptr;
+	}
+
+	FNiagaraSystemInstance* SystemInstance = Emitter->GetParentSystemInstance();
+	UWorld* World = SystemInstance ? SystemInstance->GetWorld() : nullptr;
+	ULineBatchComponent* LineBatcher = World ? World->GetLineBatcher(UWorld::ELineBatcherType::World) : nullptr;
+	if (!LineBatcher)
+	{
+		return nullptr;
+	}
+
+	TUniquePtr<FNiagaraUIRendererRenderData> RenderData(RendererProperties->CreateRenderData(*Emitter));
+	if (!RenderData.IsValid())
+	{
+		return nullptr;
+	}
+
+	FNiagaraUIRenderContext::FMaterialCacheMap TempCacheMap;
+	FNiagaraUIDebugRenderContext RenderContext(TempCacheMap , *LineBatcher);
+	RendererProperties->ExecuteRender(RenderContext, *RenderData.Get());
+
+	return nullptr;
+}
+#endif
diff --git a/Engine/Plugins/FX/NiagaraUIRenderer/Source/NiagaraUIRenderer/Public/NiagaraUIDebugRenderer.h b/Engine/Plugins/FX/NiagaraUIRenderer/Source/NiagaraUIRenderer/Public/NiagaraUIDebugRenderer.h
new file mode 100644
--- /dev/null
+++ b/Engine/Plugins/FX/NiagaraUIRenderer/Source/NiagaraUIRenderer/Public/NiagaraUIDebugRenderer.h
@@ -0,0 +1,21 @@
+// Copyright Epic Games, Inc. All Rights Reserved.
+#pragma once
+
+#include "NiagaraRendererProperties.h"
+#include "NiagaraCommon.h"
+#include "NiagaraGPUSortInfo.h"
+#include "NiagaraSystemInstanceController.h"
+#include "NiagaraUITypes.h"
+
+#if WITH_NIAGARA_RENDERER_DEBUGDRAW 
+class FNiagaraUIDebugRenderer : public FNiagaraRenderer
+{
+public:
+	explicit FNiagaraUIDebugRenderer(ERHIFeatureLevel::Type FeatureLevel, const UNiagaraRendererProperties* InProps, const FNiagaraEmitterInstance* Emitter);
+	virtual ~FNiagaraUIDebugRenderer() = default;
+	
+	//FNiagaraRenderer interface
+	virtual FNiagaraDynamicDataBase* GenerateDynamicData(const FNiagaraSceneProxy* Proxy, const UNiagaraRendererProperties* InProperties, const FNiagaraEmitterInstance* Emitter) const override;
+	//FNiagaraRenderer interface END
+};
+#endif
diff --git a/Engine/Plugins/FX/NiagaraUIRenderer/Source/NiagaraUIRenderer/Public/NiagaraUIRenderContext.h b/Engine/Plugins/FX/NiagaraUIRenderer/Source/NiagaraUIRenderer/Public/NiagaraUIRenderContext.h
--- a/Engine/Plugins/FX/NiagaraUIRenderer/Source/NiagaraUIRenderer/Public/NiagaraUIRenderContext.h
+++ b/Engine/Plugins/FX/NiagaraUIRenderer/Source/NiagaraUIRenderer/Public/NiagaraUIRenderContext.h
@@ -20,10 +20,12 @@ struct FNiagaraUIRendererScratchBuffers
 
 class FNiagaraUIRenderContext
 {
+public:
 	using FMaterialCacheMap = TMap<FObjectKey, FSlateResourceHandle>;
 
 public:
-	explicit FNiagaraUIRenderContext(FMaterialCacheMap& InCacheMap, const FGeometry& InAllottedGeometry, FSlateWindowElementList& InOutDrawElements);
+	explicit FNiagaraUIRenderContext(FMaterialCacheMap& InCacheMap, const FGeometry& InAllottedGeometry);
+	virtual ~FNiagaraUIRenderContext() = default;
 
 	// Get scratch buffers for generating vertex data into
 	FNiagaraUIRendererScratchBuffers GetScratchBuffers(int32 NumVertices, int32 NumInidices) const;
@@ -35,7 +37,7 @@ public:
 	void DrawCustomVerts(const UMaterialInterface* Material) const;
 
 	// Draw the verticies provided
-	void DrawCustomVerts(const UMaterialInterface* Material, const TArray<FSlateVertex>& Vertices, const TArray<SlateIndex>& Indices) const;
+	virtual void DrawCustomVerts(const UMaterialInterface* Material, const TArray<FSlateVertex>& Vertices, const TArray<SlateIndex>& Indices) const = 0;
 
 	// Convert Simulation Posiiton to Widget Screen
 	FVector2f PositionToScreen(const FNiagaraPosition& WorldPosition) const
@@ -59,17 +61,16 @@ public:
 	void SetScreenOriginAlignment(EHorizontalAlignment HorizontalAlignment, EVerticalAlignment VerticalAlignment);
 	void SetScreenScale(float InScreenScale) { ScreenScale = InScreenScale; }
 
-private:
+protected:
 	const FSlateResourceHandle& GetOrCreateMaterialHandle(const UMaterialInterface* Material) const;
 
-private:
+protected:
 	mutable TArray<FSlateVertex>	ScratchVertices;
 	mutable TArray<SlateIndex>		ScratchIndices;
 	mutable TArray<int32>			ParticleOrderTable;
 
 	FMaterialCacheMap&				MaterialHandleCache;
 	const FGeometry&				AllottedGeometry;
-	FSlateWindowElementList&		OutDrawElements;
 	int32							LayerId = 0;
 	ESlateDrawEffect				DrawEffect = ESlateDrawEffect::None;
 
diff --git a/Engine/Plugins/FX/NiagaraUIRenderer/Source/NiagaraUIRenderer/Public/NiagaraUIRendererProperties.cpp b/Engine/Plugins/FX/NiagaraUIRenderer/Source/NiagaraUIRenderer/Public/NiagaraUIRendererProperties.cpp
--- a/Engine/Plugins/FX/NiagaraUIRenderer/Source/NiagaraUIRenderer/Public/NiagaraUIRendererProperties.cpp
+++ b/Engine/Plugins/FX/NiagaraUIRenderer/Source/NiagaraUIRenderer/Public/NiagaraUIRendererProperties.cpp
@@ -1,10 +1,15 @@
 // Copyright Epic Games, Inc. All Rights Reserved.
 
 #include "NiagaraUIRendererProperties.h"
+#include "NiagaraUIDebugRenderer.h"
 
-FNiagaraRenderer* UNiagaraUIRendererProperties::CreateEmitterRenderer(ERHIFeatureLevel::Type /*FeatureLevel*/, const FNiagaraEmitterInstance* /*Emitter*/, const FNiagaraSystemInstanceController& /*InController*/)
+FNiagaraRenderer* UNiagaraUIRendererProperties::CreateEmitterRenderer(ERHIFeatureLevel::Type FeatureLevel, const FNiagaraEmitterInstance* Emitter, const FNiagaraSystemInstanceController& InController)
 {
+#if WITH_NIAGARA_RENDERER_DEBUGDRAW 
+	return new FNiagaraUIDebugRenderer(FeatureLevel, this, Emitter);
+#else
 	return nullptr;
+#endif
 }
 
 FNiagaraBoundsCalculator* UNiagaraUIRendererProperties::CreateBoundsCalculator()
diff --git a/Engine/Plugins/FX/NiagaraUIRenderer/Source/NiagaraUIRenderer/Public/NiagaraUIRendererProperties.h b/Engine/Plugins/FX/NiagaraUIRenderer/Source/NiagaraUIRenderer/Public/NiagaraUIRendererProperties.h
--- a/Engine/Plugins/FX/NiagaraUIRenderer/Source/NiagaraUIRenderer/Public/NiagaraUIRendererProperties.h
+++ b/Engine/Plugins/FX/NiagaraUIRenderer/Source/NiagaraUIRenderer/Public/NiagaraUIRendererProperties.h
@@ -17,9 +17,12 @@ class UNiagaraUIRendererProperties : public UNiagaraRendererProperties
 
 public:
 	//~ BEGIN: UNiagaraRendererProperties interface
-	virtual FNiagaraRenderer* CreateEmitterRenderer(ERHIFeatureLevel::Type /*FeatureLevel*/, const FNiagaraEmitterInstance* /*Emitter*/, const FNiagaraSystemInstanceController& /*InController*/) override;
-	virtual FNiagaraBoundsCalculator* CreateBoundsCalculator() override;
-	virtual bool IsSimTargetSupported(ENiagaraSimTarget InSimTarget) const override;
+	virtual FNiagaraRenderer* CreateEmitterRenderer(ERHIFeatureLevel::Type FeatureLevel, const FNiagaraEmitterInstance* Emitter, const FNiagaraSystemInstanceController& InController) final override;
+	virtual FNiagaraBoundsCalculator* CreateBoundsCalculator() final override;
+	virtual bool IsSimTargetSupported(ENiagaraSimTarget InSimTarget) const final override;
+#if WITH_NIAGARA_RENDERER_DEBUGDRAW 
+	virtual bool SupportsDebugDraw() const { return true; }
+#endif
 	//~ END: UNiagaraRendererProperties interface
 
 	virtual FNiagaraUIRendererRenderData* CreateRenderData(const FNiagaraEmitterInstance&) const { return nullptr; }
diff --git a/Engine/Plugins/FX/NiagaraUIRenderer/Source/NiagaraUIRenderer/Public/NiagaraUISpriteRendererProperties.h b/Engine/Plugins/FX/NiagaraUIRenderer/Source/NiagaraUIRenderer/Public/NiagaraUISpriteRendererProperties.h
--- a/Engine/Plugins/FX/NiagaraUIRenderer/Source/NiagaraUIRenderer/Public/NiagaraUISpriteRendererProperties.h
+++ b/Engine/Plugins/FX/NiagaraUIRenderer/Source/NiagaraUIRenderer/Public/NiagaraUISpriteRendererProperties.h
@@ -14,10 +14,7 @@ public:
 	UNiagaraUISpriteRendererProperties();
 
 	//~ BEGIN: UNiagaraRendererProperties interface
-	virtual FNiagaraRenderer* CreateEmitterRenderer(ERHIFeatureLevel::Type FeatureLevel, const FNiagaraEmitterInstance* Emitter, const FNiagaraSystemInstanceController& InController) override;
-	virtual class FNiagaraBoundsCalculator* CreateBoundsCalculator() override;
 	virtual void GetUsedMaterials(const FNiagaraEmitterInstance* EmitterInstance, TArray<UMaterialInterface*>& OutMaterials) const override;
-	virtual bool IsSimTargetSupported(ENiagaraSimTarget InSimTarget) const override;
 	virtual void PostInitProperties() override;
 	virtual void PostLoad() override;
 #if WITH_EDITORONLY_DATA
